[PCG] Python Interop Plugin and Execute Python Script Node

+ New PCGPythonInterop Plugin (Editor Only)
+ New Execute Python Script Node
    + Script from Input, Inline, or File

TODO: Optional input pins that get autoconverted to python vars

diff --git a/Engine/Plugins/PCG/Source/PCG/Private/Utils/PCGLogErrors.cpp b/Engine/Plugins/PCG/Source/PCG/Private/Utils/PCGLogErrors.cpp
--- a/Engine/Plugins/PCG/Source/PCG/Private/Utils/PCGLogErrors.cpp
+++ b/Engine/Plugins/PCG/Source/PCG/Private/Utils/PCGLogErrors.cpp
@@ -13,6 +13,9 @@ namespace PCGLog::InputOutput
 		const FTextFormat TypedInputNotFound = LOCTEXT("TypedInputNotFound", "Data of type {0} not found on pin '{1}'.");
 		const FTextFormat FirstInputOnly = LOCTEXT("FirstInputOnly", "Multiple inputs found on single-input pin '{0}'. Only the first will be selected.");
 		const FTextFormat InvalidCardinality = LOCTEXT("InvalidCardinality", "Invalid cardinality among pins '{0}' and '{1}'. They must match 1:1, N:1, or N:N.");
+		const FTextFormat FileNotFound = LOCTEXT("FileNotFound", "File was not found at given path: '{0}'");
+		const FTextFormat InvalidFileType = LOCTEXT("InvalidFileType", "File was of the wrong type: '{0}' Expected type: {1}");
+		const FTextFormat DirectoryNotFound = LOCTEXT("DirectoryNotFound", "Directory was not found at given path: '{0}'");
 	}
 
 	// Warnings
@@ -38,6 +41,21 @@ namespace PCGLog::InputOutput
 	{
 		LogErrorOnGraph(FText::Format(Format::InvalidCardinality, FText::FromName(SourcePinLabel), FText::FromName(TargetPinLabel)), InContext);
 	}
+
+	void LogFileNotFound(const FString& FilePath, const FPCGContext* InContext)
+	{
+		LogErrorOnGraph(FText::Format(Format::FileNotFound, FText::FromString(FilePath)), InContext);
+	}
+
+	void LogInvalidFileType(const FString& FileType, const FString& ExpectedType, const FPCGContext* InContext)
+	{
+		LogErrorOnGraph(FText::Format(Format::InvalidFileType, FText::FromString(FileType), FText::FromString(ExpectedType)), InContext);
+	}
+
+	void LogDirectoryNotFound(const FString& DirectoryPath, const FPCGContext* InContext)
+	{
+		LogErrorOnGraph(FText::Format(Format::DirectoryNotFound, FText::FromString(DirectoryPath)), InContext);
+	}
 }
 
 namespace PCGLog::Metadata
diff --git a/Engine/Plugins/PCG/Source/PCG/Public/Metadata/PCGDefaultValueInterface.h b/Engine/Plugins/PCG/Source/PCG/Public/Metadata/PCGDefaultValueInterface.h
--- a/Engine/Plugins/PCG/Source/PCG/Public/Metadata/PCGDefaultValueInterface.h
+++ b/Engine/Plugins/PCG/Source/PCG/Public/Metadata/PCGDefaultValueInterface.h
@@ -9,6 +9,16 @@
 
 #include "PCGDefaultValueInterface.generated.h"
 
+UENUM()
+enum class EPCGSettingDefaultValueExtraFlags : uint64
+{
+	None          = 0,
+	WideText      = 1 << 0, // The text box will need to be large enough to support a long text block.
+	MultiLineText = 1 << 1  // The text block may span multiple lines.
+};
+
+ENUM_CLASS_FLAGS(EPCGSettingDefaultValueExtraFlags);
+
 UINTERFACE(MinimalAPI)
 class UPCGSettingsDefaultValueProvider : public UInterface
 {
@@ -67,6 +77,9 @@ public:
 
 	/** For the initial 'default value' of the pin. */
 	virtual FString GetPinInitialDefaultValueString(FName PinLabel) const { return FString(); }
+
+	/** Extra flags related to displaying the value. */
+	virtual EPCGSettingDefaultValueExtraFlags GetDefaultValueExtraFlags(FName PinLabel) const { return EPCGSettingDefaultValueExtraFlags::None; }
 #endif // WITH_EDITOR
 
 protected:
diff --git a/Engine/Plugins/PCG/Source/PCG/Public/Utils/PCGLogErrors.h b/Engine/Plugins/PCG/Source/PCG/Public/Utils/PCGLogErrors.h
--- a/Engine/Plugins/PCG/Source/PCG/Public/Utils/PCGLogErrors.h
+++ b/Engine/Plugins/PCG/Source/PCG/Public/Utils/PCGLogErrors.h
@@ -19,9 +19,12 @@ namespace PCGLog
 	{
 		namespace Format
 		{
+			PCG_API extern const FText InvalidInputData;
+
 			PCG_API extern const FTextFormat TypedInputNotFound;
 			PCG_API extern const FTextFormat FirstInputOnly;
-			PCG_API extern const FText InvalidInputData;
+			PCG_API extern const FTextFormat FileNotFound;
+			PCG_API extern const FTextFormat DirectoryNotFound;
 		}
 
 		// Warnings
@@ -31,6 +34,9 @@ namespace PCGLog
 		// Errors
 		PCG_API void LogInvalidInputDataError(const FPCGContext* InContext = nullptr);
 		PCG_API void LogInvalidCardinalityError(const FName SourcePinLabel, const FName TargetPinLabel, const FPCGContext* InContext = nullptr);
+		PCG_API void LogFileNotFound(const FString& FilePath, const FPCGContext* InContext = nullptr);
+		PCG_API void LogInvalidFileType(const FString& FileType, const FString& ExpectedType, const FPCGContext* InContext = nullptr);
+		PCG_API void LogDirectoryNotFound(const FString& DirectoryPath, const FPCGContext* InContext = nullptr);
 	}
 
 	namespace Metadata
diff --git a/Engine/Plugins/PCG/Source/PCGEditor/Private/Nodes/SPCGEditorGraphNode.cpp b/Engine/Plugins/PCG/Source/PCGEditor/Private/Nodes/SPCGEditorGraphNode.cpp
--- a/Engine/Plugins/PCG/Source/PCGEditor/Private/Nodes/SPCGEditorGraphNode.cpp
+++ b/Engine/Plugins/PCG/Source/PCGEditor/Private/Nodes/SPCGEditorGraphNode.cpp
@@ -134,25 +134,34 @@ TSharedPtr<SGraphPin> SPCGEditorGraphNode::CreatePinWidget(UEdGraphPin* InPin) c
 			{
 				case EPCGMetadataTypes::Name: // fall-through
 				case EPCGMetadataTypes::String:
-					return SNew(SPCGEditorGraphPinString, InPin, std::move(OnModify));
+				{
+					const EPCGSettingDefaultValueExtraFlags Flags = DefaultValueInterface->GetDefaultValueExtraFlags(InPin->PinName);
+					const bool bIsWide = EnumHasAnyFlags(Flags, EPCGSettingDefaultValueExtraFlags::WideText);
+					const bool bIsMultiLine = EnumHasAnyFlags(Flags, EPCGSettingDefaultValueExtraFlags::MultiLineText);
+					return SNew(SPCGEditorGraphPinString, InPin, MoveTemp(OnModify))
+						.MinDesiredBoxWidth(bIsWide ? 150.f : 60.f)
+						.MaxDesiredBoxWidth(bIsWide ? 600.f : 200.f)
+						.IsMultiline(bIsMultiLine)
+						.OverflowPolicy(bIsMultiLine ? ETextOverflowPolicy::MultilineEllipsis : ETextOverflowPolicy::Ellipsis);
+				}
 				// Float is converted to double by the property accessor under the hood
 				case EPCGMetadataTypes::Float: // fall-through
 				case EPCGMetadataTypes::Double:
-					return SNew(SPCGEditorGraphPinNumSlider<double>, InPin, std::move(OnModify));
+					return SNew(SPCGEditorGraphPinNumSlider<double>, InPin, MoveTemp(OnModify));
 				case EPCGMetadataTypes::Integer32:
-					return SNew(SPCGEditorGraphPinNumSlider<int32>, InPin, std::move(OnModify))
+					return SNew(SPCGEditorGraphPinNumSlider<int32>, InPin, MoveTemp(OnModify))
 						.MinDesiredBoxWidth(40.f);
 				case EPCGMetadataTypes::Integer64:
-					return SNew(SPCGEditorGraphPinNumSlider<int64>, InPin, std::move(OnModify))
+					return SNew(SPCGEditorGraphPinNumSlider<int64>, InPin, MoveTemp(OnModify))
 						.MinDesiredBoxWidth(40.f);
 				case EPCGMetadataTypes::Vector:
-					return SNew(SPCGEditorGraphPinVectorSlider<FVector>, InPin, std::move(OnModify));
+					return SNew(SPCGEditorGraphPinVectorSlider<FVector>, InPin, MoveTemp(OnModify));
 				case EPCGMetadataTypes::Vector2:
-					return SNew(SPCGEditorGraphPinVectorSlider<FVector2D>, InPin, std::move(OnModify));
+					return SNew(SPCGEditorGraphPinVectorSlider<FVector2D>, InPin, MoveTemp(OnModify));
 				case EPCGMetadataTypes::Vector4:
-					return SNew(SPCGEditorGraphPinVectorSlider<FVector4>, InPin, std::move(OnModify));
+					return SNew(SPCGEditorGraphPinVectorSlider<FVector4>, InPin, MoveTemp(OnModify));
 				case EPCGMetadataTypes::Rotator:
-					return SNew(SPCGEditorGraphPinVectorSlider<FRotator>, InPin, std::move(OnModify));
+					return SNew(SPCGEditorGraphPinVectorSlider<FRotator>, InPin, MoveTemp(OnModify));
 				case EPCGMetadataTypes::Boolean:
 					return SNew(SPCGEditorGraphPinBool, InPin, std::move(OnModify));
 				// @todo_pcg: Will be added once widgets are created.
diff --git a/Engine/Plugins/PCG/Source/PCGEditor/Private/PCGEditorCommon.cpp b/Engine/Plugins/PCG/Source/PCGEditor/Private/PCGEditorCommon.cpp
--- a/Engine/Plugins/PCG/Source/PCGEditor/Private/PCGEditorCommon.cpp
+++ b/Engine/Plugins/PCG/Source/PCGEditor/Private/PCGEditorCommon.cpp
@@ -4,10 +4,23 @@
 
 #include UE_INLINE_GENERATED_CPP_BY_NAME(PCGEditorCommon)
 
+#include "Framework/Notifications/NotificationManager.h"
+#include "Widgets/Notifications/SNotificationList.h"
+
 namespace FPCGEditorCommon
 {
 	TAutoConsoleVariable<bool> CVarShowAdvancedAttributesFields(
 		TEXT("pcg.graph.ShowAdvancedAttributes"),
 		false,
 		TEXT("Control whether advanced attributes/properties are shown in the PCG graph editor"));
+
+	void Helpers::DispatchEditorToast(const FText& Text, const FText& SubText, const float Duration)
+	{
+		FNotificationInfo ToastInfo(Text);
+		ToastInfo.ExpireDuration = Duration;
+		ToastInfo.bFireAndForget = true;
+		ToastInfo.SubText = SubText;
+
+		FSlateNotificationManager::Get().AddNotification(ToastInfo);
+	}
 }
diff --git a/Engine/Plugins/PCG/Source/PCGEditor/Private/Pins/SPCGEditorGraphPinString.h b/Engine/Plugins/PCG/Source/PCGEditor/Private/Pins/SPCGEditorGraphPinString.h
--- a/Engine/Plugins/PCG/Source/PCGEditor/Private/Pins/SPCGEditorGraphPinString.h
+++ b/Engine/Plugins/PCG/Source/PCGEditor/Private/Pins/SPCGEditorGraphPinString.h
@@ -7,7 +7,7 @@
 #include "ScopedTransaction.h"
 
 #include "SSearchableComboBox.h"
-#include "Widgets/Input/SEditableTextBox.h"
+#include "Widgets/Input/SMultiLineEditableTextBox.h"
 
 /** Note: Derived from Engine/Source/Editor/GraphEditor/Public/KismetPins/SGraphPinString.h */
 
@@ -16,11 +16,14 @@ class SPCGEditorGraphPinString final : public SPCGEditorGraphNodePin
 public:
 	SLATE_BEGIN_ARGS(SPCGEditorGraphPinString)
 		: _MinDesiredBoxWidth(60.f)
-		, _MaxDesiredBoxWidth(400.f)
-	{}
+		, _MaxDesiredBoxWidth(300.f)
+		, _IsMultiline(false)
+		, _OverflowPolicy(ETextOverflowPolicy::Ellipsis) {}
 
 	SLATE_ARGUMENT(float, MinDesiredBoxWidth)
 	SLATE_ARGUMENT(float, MaxDesiredBoxWidth)
+	SLATE_ARGUMENT(bool, IsMultiline)
+	SLATE_ARGUMENT(ETextOverflowPolicy, OverflowPolicy)
 
 	SLATE_END_ARGS()
 
@@ -29,6 +32,8 @@ public:
 		OnModifyDelegate = std::move(OnModify);
 		MinDesiredBoxWidth = InArgs._MinDesiredBoxWidth;
 		MaxDesiredBoxWidth = InArgs._MaxDesiredBoxWidth;
+		bIsMultiline = InArgs._IsMultiline;
+		OverflowPolicy = InArgs._OverflowPolicy;
 
 		SPCGEditorGraphNodePin::Construct(SPCGEditorGraphNodePin::FArguments(), InPin);
 	}
@@ -40,9 +45,11 @@ public:
 			.MinDesiredWidth(MinDesiredBoxWidth)
 			.MaxDesiredWidth(MaxDesiredBoxWidth)
 			[
-				SNew(SEditableTextBox)
+				SNew(SMultiLineEditableTextBox)
 				.Style(FAppStyle::Get(), "Graph.EditableTextBox")
 				.Text(this, &SPCGEditorGraphPinString::GetTypeInValue)
+				.AllowMultiLine(bIsMultiline)
+				.OverflowPolicy(OverflowPolicy)
 				.SelectAllTextWhenFocused(true)
 				.Visibility(this, &SGraphPin::GetDefaultValueVisibility)
 				.IsReadOnly(this, &SPCGEditorGraphPinString::GetDefaultValueIsReadOnly)
@@ -83,5 +90,7 @@ private:
 	TArray<TSharedPtr<FString>> ComboBoxOptions;
 	float MinDesiredBoxWidth = 0;
 	float MaxDesiredBoxWidth = 0;
+	bool bIsMultiline = false;
+	ETextOverflowPolicy OverflowPolicy = ETextOverflowPolicy::Clip;
 	FSimpleDelegate OnModifyDelegate;
 };
diff --git a/Engine/Plugins/PCG/Source/PCGEditor/Private/Widgets/SPCGEditorGraphDetailsView.cpp b/Engine/Plugins/PCG/Source/PCGEditor/Private/Widgets/SPCGEditorGraphDetailsView.cpp
--- a/Engine/Plugins/PCG/Source/PCGEditor/Private/Widgets/SPCGEditorGraphDetailsView.cpp
+++ b/Engine/Plugins/PCG/Source/PCGEditor/Private/Widgets/SPCGEditorGraphDetailsView.cpp
@@ -43,19 +43,6 @@ namespace PCGEditorGraphDetailsView
 		static const FText DefaultDebugMaterialDensityButtonTooltip = LOCTEXT("DefaultDebugMaterialDensityButtonTooltip", "Select the default debug material - 'PCG_DebugMaterial'.");
 		static const FText DefaultDebugMaterialColorButtonTooltip = LOCTEXT("DefaultDebugMaterialColorButtonTooltip", "Select the default  debug material - 'PCG_DebugMaterial_Color'");
 	}
-
-	namespace Helpers
-	{
-		void DispatchEditorToast(const FText& Text, const FText& SubText = FText{}, const float Duration = 3.f)
-		{
-			FNotificationInfo ToastInfo(Text);
-			ToastInfo.ExpireDuration = Duration;
-			ToastInfo.bFireAndForget = true;
-			ToastInfo.SubText = SubText;
-
-			FSlateNotificationManager::Get().AddNotification(ToastInfo);
-		}
-	}
 }
 
 class FPCGDebugDetailsExtension : public IDetailPropertyExtensionHandler
@@ -104,7 +91,7 @@ class FPCGDebugDetailsExtension : public IDetailPropertyExtensionHandler
 						.OnClicked_Lambda([PropertyHandle]
 						{
 							PropertyHandle->SetValueFromFormattedString(Constants::DefaultDebugDensityMaterialPath);
-							Helpers::DispatchEditorToast(Constants::DebugMaterialButtonToastText, Constants::DebugMaterialDensityButtonToastSubText);
+							FPCGEditorCommon::Helpers::DispatchEditorToast(Constants::DebugMaterialButtonToastText, Constants::DebugMaterialDensityButtonToastSubText);
 							return FReply::Handled();
 						})
 						.ToolTipText(Constants::DefaultDebugMaterialDensityButtonTooltip)
@@ -125,7 +112,7 @@ class FPCGDebugDetailsExtension : public IDetailPropertyExtensionHandler
 						.OnClicked_Lambda([PropertyHandle]
 						{
 							PropertyHandle->SetValueFromFormattedString(Constants::DefaultDebugMaterialColorPath);
-							Helpers::DispatchEditorToast(Constants::DebugMaterialButtonToastText, Constants::DebugMaterialColorButtonToastSubText);
+							FPCGEditorCommon::Helpers::DispatchEditorToast(Constants::DebugMaterialButtonToastText, Constants::DebugMaterialColorButtonToastSubText);
 
 							return FReply::Handled();
 						})
diff --git a/Engine/Plugins/PCG/Source/PCGEditor/Public/PCGEditorCommon.h b/Engine/Plugins/PCG/Source/PCGEditor/Public/PCGEditorCommon.h
--- a/Engine/Plugins/PCG/Source/PCGEditor/Public/PCGEditorCommon.h
+++ b/Engine/Plugins/PCG/Source/PCGEditor/Public/PCGEditorCommon.h
@@ -56,6 +56,11 @@ namespace FPCGEditorCommon
 
 	const FAssetCategoryPath PCGAssetCategoryPath(LOCTEXT("PCGCategory", "PCG"));
 	const FAssetCategoryPath PCGAdvancedAssetCategoryPath(FPCGEditorCommon::PCGAssetCategoryPath, LOCTEXT("PCGAvancedCategory", "Advanced"));
-};
+
+	namespace Helpers
+	{
+		PCGEDITOR_API void DispatchEditorToast(const FText& Text, const FText& SubText = FText{}, const float Duration = 3.f);
+	}
+}
 
 #undef LOCTEXT_NAMESPACE
diff --git a/Engine/Plugins/PCGInterops/PCGPythonInterop/PCGPythonInterop.uplugin b/Engine/Plugins/PCGInterops/PCGPythonInterop/PCGPythonInterop.uplugin
new file mode 100644
--- /dev/null
+++ b/Engine/Plugins/PCGInterops/PCGPythonInterop/PCGPythonInterop.uplugin
@@ -0,0 +1,36 @@
+{
+    "FileVersion": 1,
+    "Version": 1,
+    "VersionName": "0.1",
+    "FriendlyName": "Procedural Content Generation Framework (PCG) Python Interop",
+    "Description": "Extra plugin for Procedural Content Generation Framework interacting with the Editor Python Interpreter.",
+    "Category": "Editor",
+    "CreatedBy": "Epic Games, Inc.",
+    "CreatedByURL": "https://epicgames.com",
+    "DocsURL": "https://docs.unrealengine.com/latest/en-US/procedural-content-generation--framework-in-unreal-engine/",
+    "MarketplaceURL": "",
+    "SupportURL": "",
+    "EnabledByDefault": false,
+    "CanContainContent": true,
+    "IsBetaVersion": true,
+    "Installed": false,
+    "Modules": [
+        {
+            "Name": "PCGPythonInteropEditor",
+            "Type": "Editor",
+            "LoadingPhase": "Default"
+        }
+    ],
+    "Plugins": [
+        {
+            "Name": "PCG",
+            "Enabled": true,
+	        "Optional" : false
+        },
+        {
+            "Name": "PythonScriptPlugin",
+            "Enabled": true,
+	        "Optional" : false
+        }
+    ]
+}
diff --git a/Engine/Plugins/PCGInterops/PCGPythonInterop/Source/PCGPythonInteropEditor/PCGPythonInteropEditor.Build.cs b/Engine/Plugins/PCGInterops/PCGPythonInterop/Source/PCGPythonInteropEditor/PCGPythonInteropEditor.Build.cs
new file mode 100644
--- /dev/null
+++ b/Engine/Plugins/PCGInterops/PCGPythonInterop/Source/PCGPythonInteropEditor/PCGPythonInteropEditor.Build.cs
@@ -0,0 +1,35 @@
+// Copyright Epic Games, Inc. All Rights Reserved.
+
+namespace UnrealBuildTool.Rules
+{
+	public class PCGPythonInteropEditor : ModuleRules
+	{
+		public PCGPythonInteropEditor(ReadOnlyTargetRules Target) : base(Target)
+		{
+			PCHUsage = PCHUsageMode.UseExplicitOrSharedPCHs;
+
+			PublicDependencyModuleNames.AddRange(
+				new string[] {
+					"Core",
+					"CoreUObject",
+					"Engine",
+					"Projects",
+					"PCG"
+				});
+
+			if (Target.bBuildEditor)
+			{
+				PrivateDependencyModuleNames.AddRange(
+					new string[]
+					{
+						"AssetDefinition",
+						"PCGEditor",
+						"UnrealEd",
+						"PythonScriptPlugin",
+						"Slate"
+					}
+				);
+			}
+		}
+	}
+}
diff --git a/Engine/Plugins/PCGInterops/PCGPythonInterop/Source/PCGPythonInteropEditor/Private/Elements/PCGExecutePythonScript.cpp b/Engine/Plugins/PCGInterops/PCGPythonInterop/Source/PCGPythonInteropEditor/Private/Elements/PCGExecutePythonScript.cpp
new file mode 100644
--- /dev/null
+++ b/Engine/Plugins/PCGInterops/PCGPythonInterop/Source/PCGPythonInteropEditor/Private/Elements/PCGExecutePythonScript.cpp
@@ -0,0 +1,359 @@
+// Copyright Epic Games, Inc. All Rights Reserved.
+
+#include "Elements/PCGExecutePythonScript.h"
+
+#include "PCGContext.h"
+#include "PCGParamData.h"
+#include "Metadata/Accessors/PCGAttributeAccessorHelpers.h"
+#include "Utils/PCGLogErrors.h"
+
+#include "PCGEditorCommon.h"
+#include "PCGPythonInteropEditorModule.h"
+#include "PCGGraph.h"
+
+#include "IPythonScriptPlugin.h"
+#include "HAL/FileManager.h"
+#include "Misc/Paths.h"
+
+#define LOCTEXT_NAMESPACE "PCGExecutePythonScriptElement"
+
+namespace PCGExecutePythonScript::Constants
+{
+	static const FName ScriptSourceInputPinLabel = "Source";
+	static const FText ScriptSourceInputPinTooltip = LOCTEXT("ScriptSourcePinTooltip", "The script may be passed in as an FString attribute input.");
+	static constexpr TCHAR ExpectedFileExtension[] = TEXT("py");
+}
+
+const FString UPCGExecutePythonScriptSettings::DefaultInlineScript = LOCTEXT("InlineScriptDefault", "print(\"Hello PCG World!\")").ToString();
+
+#if WITH_EDITOR
+bool UPCGExecutePythonScriptSettings::CanEditChange(const FProperty* InProperty) const
+{
+	using namespace PCGExecutePythonScript::Constants;
+
+	if (!InProperty || !Super::CanEditChange(InProperty))
+	{
+		return false;
+	}
+
+	if (ScriptInputMethod == EPCGPythonScriptInputMethod::Input
+		&& InProperty->GetFName() == GET_MEMBER_NAME_CHECKED(UPCGExecutePythonScriptSettings, ScriptSource))
+	{
+		if (const UPCGNode* Node = Cast<UPCGNode>(GetOuter()))
+		{
+			return Node->GetInputPin(ScriptSourceInputPinLabel)->IsConnected() || !IsPinDefaultValueActivated(ScriptSourceInputPinLabel);
+		}
+	}
+
+	return true;
+}
+
+EPCGChangeType UPCGExecutePythonScriptSettings::GetChangeTypeForProperty(const FName& InPropertyName) const
+{
+	EPCGChangeType ChangeType = Super::GetChangeTypeForProperty(InPropertyName);
+
+	if (InPropertyName == GET_MEMBER_NAME_CHECKED(UPCGExecutePythonScriptSettings, ScriptPath))
+	{
+		ChangeType |= EPCGChangeType::Settings | EPCGChangeType::Structural | EPCGChangeType::Cosmetic;
+	}
+
+	return ChangeType;
+}
+
+FName UPCGExecutePythonScriptSettings::GetDefaultNodeName() const
+{
+	return FName(TEXT("ExecutePythonScript"));
+}
+
+FText UPCGExecutePythonScriptSettings::GetDefaultNodeTitle() const
+{
+	return LOCTEXT("NodeTitle", "Execute Python Script");
+}
+
+FText UPCGExecutePythonScriptSettings::GetNodeTooltipText() const
+{
+	return LOCTEXT("NodeTooltip", "Execute a Python script from an inline or connected input, or directly from a .py file.");
+}
+
+EPCGSettingsType UPCGExecutePythonScriptSettings::GetType() const
+{
+	return EPCGSettingsType::Generic;
+}
+
+FString UPCGExecutePythonScriptSettings::GetAdditionalTitleInformation() const
+{
+	switch (ScriptInputMethod)
+	{
+		case EPCGPythonScriptInputMethod::Input:
+			return LOCTEXT("InputTitleInformation", "Source Input").ToString();
+		case EPCGPythonScriptInputMethod::File:
+		{
+			return FText::Format(
+				LOCTEXT("FileTitleInformation", "Source File: {0}"),
+				FText::FromString(FPaths::GetCleanFilename(ScriptPath.FilePath))).ToString();
+		}
+	}
+
+	return {};
+}
+
+TArray<FPCGPinProperties> UPCGExecutePythonScriptSettings::InputPinProperties() const
+{
+	if (ScriptInputMethod == EPCGPythonScriptInputMethod::Input)
+	{
+		TArray<FPCGPinProperties> PinProperties;
+		FPCGPinProperties& ScriptSourcePin = PinProperties.Emplace_GetRef(PCGExecutePythonScript::Constants::ScriptSourceInputPinLabel, EPCGDataType::Param);
+#if WITH_EDITOR
+		ScriptSourcePin.Tooltip = PCGExecutePythonScript::Constants::ScriptSourceInputPinTooltip;
+#endif // WITH_EDITOR
+		return PinProperties;
+	}
+	else
+	{
+		return {};
+	}
+}
+
+TArray<FPCGPinProperties> UPCGExecutePythonScriptSettings::OutputPinProperties() const
+{
+	TArray<FPCGPinProperties> PinProperties;
+	FPCGPinProperties& DepPin = PinProperties.Emplace_GetRef(PCGPinConstants::DefaultExecutionDependencyLabel, EPCGDataType::Any);
+#if WITH_EDITOR
+	DepPin.Tooltip = PCGPinConstants::Tooltips::ExecutionDependencyTooltip;
+#endif // WITH_EDITOR
+	DepPin.Usage = EPCGPinUsage::DependencyOnly;
+
+	return PinProperties;
+}
+
+#endif // WITH_EDITOR
+
+FPCGElementPtr UPCGExecutePythonScriptSettings::CreateElement() const
+{
+	return MakeShared<FPCGExecutePythonScriptElement>();
+}
+
+#if WITH_EDITOR
+bool UPCGExecutePythonScriptSettings::DefaultValuesAreEnabled() const
+{
+	return ScriptInputMethod == EPCGPythonScriptInputMethod::Input;
+}
+
+bool UPCGExecutePythonScriptSettings::IsPinDefaultValueEnabled(const FName PinLabel) const
+{
+	return PinLabel == PCGExecutePythonScript::Constants::ScriptSourceInputPinLabel;
+}
+
+bool UPCGExecutePythonScriptSettings::IsPinDefaultValueActivated(const FName PinLabel) const
+{
+	return PinLabel == PCGExecutePythonScript::Constants::ScriptSourceInputPinLabel && bIsDefaultValueActivated;
+}
+
+EPCGMetadataTypes UPCGExecutePythonScriptSettings::GetPinDefaultValueType(const FName PinLabel) const
+{
+	return EPCGMetadataTypes::String;
+}
+
+bool UPCGExecutePythonScriptSettings::IsPinDefaultValueMetadataTypeValid(const FName PinLabel, const EPCGMetadataTypes DataType) const
+{
+	return DataType == EPCGMetadataTypes::String;
+}
+
+bool UPCGExecutePythonScriptSettings::CreateInitialDefaultValueAttribute(const FName PinLabel, UPCGMetadata* OutMetadata) const
+{
+	check(PinLabel == PCGExecutePythonScript::Constants::ScriptSourceInputPinLabel);
+	return OutMetadata && OutMetadata->CreateAttribute<FString>(NAME_None, DefaultInlineScript, /*bAllowsInterpolation=*/false, /*bOverrideParent=*/false);
+}
+
+bool UPCGExecutePythonScriptSettings::IsInputPinRequiredByExecution(const UPCGPin* InPin) const
+{
+	// If the default value is not activated, or the pin is connected, it becomes required.
+	return InPin && (!IsPinDefaultValueActivated(InPin->Properties.Label) || InPin->IsConnected());
+}
+
+void UPCGExecutePythonScriptSettings::ResetDefaultValues()
+{
+	ResetDefaultValue(PCGExecutePythonScript::Constants::ScriptSourceInputPinLabel);
+}
+
+void UPCGExecutePythonScriptSettings::ResetDefaultValue(const FName PinLabel)
+{
+	check(PinLabel == PCGExecutePythonScript::Constants::ScriptSourceInputPinLabel);
+	SetPinDefaultValue(PinLabel, DefaultInlineScript);
+}
+
+void UPCGExecutePythonScriptSettings::SetPinDefaultValue(const FName PinLabel, const FString& DefaultValue, const bool bCreateIfNeeded)
+{
+	check(PinLabel == PCGExecutePythonScript::Constants::ScriptSourceInputPinLabel);
+	Modify();
+	InlineScript = DefaultValue;
+	OnSettingsChangedDelegate.Broadcast(this, EPCGChangeType::Node | EPCGChangeType::Edge);
+}
+
+void UPCGExecutePythonScriptSettings::SetPinDefaultValueIsActivated(const FName PinLabel, const bool bIsActivated, const bool bDirtySettings)
+{
+	check(PinLabel == PCGExecutePythonScript::Constants::ScriptSourceInputPinLabel);
+
+	if (bIsDefaultValueActivated != bIsActivated)
+	{
+		if (bDirtySettings)
+		{
+			Modify();
+		}
+
+		bIsDefaultValueActivated = bIsActivated;
+
+		if (bDirtySettings)
+		{
+			OnSettingsChangedDelegate.Broadcast(this, EPCGChangeType::Node | EPCGChangeType::Edge);
+		}
+	}
+}
+
+FString UPCGExecutePythonScriptSettings::GetPinDefaultValueAsString(const FName PinLabel) const
+{
+	check(PinLabel == PCGExecutePythonScript::Constants::ScriptSourceInputPinLabel);
+	return InlineScript;
+}
+
+FString UPCGExecutePythonScriptSettings::GetPinInitialDefaultValueString(const FName PinLabel) const
+{
+	check(PinLabel == PCGExecutePythonScript::Constants::ScriptSourceInputPinLabel);
+	return DefaultInlineScript;
+}
+
+EPCGSettingDefaultValueExtraFlags UPCGExecutePythonScriptSettings::GetDefaultValueExtraFlags(const FName PinLabel) const
+{
+	check(PinLabel == PCGExecutePythonScript::Constants::ScriptSourceInputPinLabel);
+	return EPCGSettingDefaultValueExtraFlags::WideText | EPCGSettingDefaultValueExtraFlags::MultiLineText;
+}
+#endif // WITH_EDITOR
+
+EPCGMetadataTypes UPCGExecutePythonScriptSettings::GetPinInitialDefaultValueType(const FName PinLabel) const
+{
+	check(PinLabel == PCGExecutePythonScript::Constants::ScriptSourceInputPinLabel);
+	return EPCGMetadataTypes::String;
+}
+
+bool FPCGExecutePythonScriptElement::ExecuteInternal(FPCGContext* InContext) const
+{
+	using namespace PCGExecutePythonScript::Constants;
+	TRACE_CPUPROFILER_EVENT_SCOPE(FPCGExecutePythonScriptElement::Execute);
+	check(InContext);
+
+#ifdef WITH_EDITOR
+	const UPCGExecutePythonScriptSettings* Settings = InContext->GetInputSettings<UPCGExecutePythonScriptSettings>();
+	check(Settings);
+
+	const IPythonScriptPlugin* PythonModule = IPythonScriptPlugin::Get();
+	if (!PythonModule || !PythonModule->IsPythonAvailable())
+	{
+		PCGLog::LogErrorOnGraph(LOCTEXT("PythonNotAvailable", "Python script can not be run. Python Interpreter is not available."), InContext);
+		// @todo_pcg: Could try to load or delay execution until its available. Would need max wait time user setting.
+		return true;
+	}
+
+	FPythonCommandEx PythonSource
+	{
+		.Flags = EPythonCommandFlags::None,
+	};
+
+	FStringBuilderBase CommandBuilder;
+
+	if (Settings->ScriptInputMethod == EPCGPythonScriptInputMethod::Input)
+	{
+		if (InContext->InputData.GetInputCountByPin(ScriptSourceInputPinLabel) > 0)
+		{
+			TArray<FPCGTaggedData> InputData = InContext->InputData.GetInputsByPin(ScriptSourceInputPinLabel);
+			for (const FPCGTaggedData& Input : InputData)
+			{
+				FPCGAttributePropertyInputSelector Selector = Settings->ScriptSource.CopyAndFixLast(Input.Data);
+				const TUniquePtr<const IPCGAttributeAccessor> Accessor = PCGAttributeAccessorHelpers::CreateConstAccessor(Input.Data, Selector);
+				const TUniquePtr<const IPCGAttributeAccessorKeys> Keys = PCGAttributeAccessorHelpers::CreateConstKeys(Input.Data, Selector);
+				TArray<FString> Values;
+				Values.SetNum(Keys->GetNum());
+				Accessor->GetRange(MakeArrayView(Values), 0, *Keys);
+				for (const FString& Value : Values)
+				{
+					CommandBuilder += Value;
+					CommandBuilder += LINE_TERMINATOR;
+				}
+			}
+		}
+		else if (Settings->IsPinDefaultValueActivated(ScriptSourceInputPinLabel))
+		{
+			CommandBuilder = Settings->GetPinDefaultValueAsString(ScriptSourceInputPinLabel);
+		}
+		else // No work to do
+		{
+			return true;
+		}
+
+		PythonSource.ExecutionMode = EPythonCommandExecutionMode::ExecuteStatement;
+	}
+	else // ScriptInputMethod == EPCGPythonScriptInputMethod::File
+	{
+		if (Settings->ScriptPath.FilePath.IsEmpty())
+		{
+			return true;
+		}
+
+		if (const FString Extension = FPaths::GetExtension(Settings->ScriptPath.FilePath); !Extension.Equals(ExpectedFileExtension))
+		{
+			const FString DisplayExtension = Extension.IsEmpty() ? LOCTEXT("NoExtension", "None").ToString() : FString::Printf(TEXT(".%s"), *Extension);
+			PCGLog::InputOutput::LogInvalidFileType(DisplayExtension, ExpectedFileExtension, InContext);
+			return true;
+		}
+
+		if (!IFileManager::Get().FileExists(*Settings->ScriptPath.FilePath))
+		{
+			PCGLog::InputOutput::LogFileNotFound(Settings->ScriptPath.FilePath, InContext);
+			return true;
+		}
+
+		CommandBuilder = Settings->ScriptPath.FilePath;
+	}
+
+	PythonSource.Command = CommandBuilder.ToString();
+
+	if (!IPythonScriptPlugin::Get()->ExecPythonCommandEx(PythonSource))
+	{
+		PCGLog::LogErrorOnGraph(LOCTEXT("PythonScriptHadErrors", "Errors running Python script. See log."));
+		for (const FPythonLogOutputEntry& LogOutputEntry : PythonSource.LogOutput)
+		{
+			switch (LogOutputEntry.Type)
+			{
+				case EPythonLogOutputType::Info:
+					UE_LOG(LogPCGPython, Log, TEXT("Python [Log]: [%s]"), *LogOutputEntry.Output);
+					break;
+				case EPythonLogOutputType::Warning:
+					UE_LOG(LogPCGPython, Warning, TEXT("Python [Warning]: [%s]"), *LogOutputEntry.Output);
+					break;
+				case EPythonLogOutputType::Error:
+					UE_LOG(LogPCGPython, Error, TEXT("Python [Error]: [%s]"), *LogOutputEntry.Output);
+					break;
+			}
+		}
+	}
+
+	// Notify the user that a script was run.
+	check(InContext->Node);
+	check(InContext->Node->GetGraph());
+
+	if (!Settings->bMuteEditorToast)
+	{
+		FPCGEditorCommon::Helpers::DispatchEditorToast(
+			LOCTEXT("PythonScriptExecutionToast", "[PCG] Python"),
+			FText::Format(INVTEXT("{0}: {1}"),
+				FText::FromString(InContext->Node->GetGraph()->GetName()),
+				InContext->Node->GetNodeTitle(EPCGNodeTitleType::ListView)));
+	}
+
+#else // WITH_EDITOR
+	UE_LOG(LogPCG, Error, LOCTEXT("NodeIsEditorOnly", "The Execute Python Script node is currently Editor-only and should not be used at runtime."));
+#endif // WITH_EDITOR
+
+	return true;
+}
+
+#undef LOCTEXT_NAMESPACE
diff --git a/Engine/Plugins/PCGInterops/PCGPythonInterop/Source/PCGPythonInteropEditor/Private/PCGPythonInteropEditorModule.cpp b/Engine/Plugins/PCGInterops/PCGPythonInterop/Source/PCGPythonInteropEditor/Private/PCGPythonInteropEditorModule.cpp
new file mode 100644
--- /dev/null
+++ b/Engine/Plugins/PCGInterops/PCGPythonInterop/Source/PCGPythonInteropEditor/Private/PCGPythonInteropEditorModule.cpp
@@ -0,0 +1,19 @@
+// Copyright Epic Games, Inc. All Rights Reserved.
+
+#include "PCGPythonInteropEditorModule.h"
+#include "Modules/ModuleManager.h"
+
+class FPCGPythonInteropEditorModule final : public IModuleInterface
+{
+public:
+	//~ IModuleInterface implementation
+	virtual bool SupportsDynamicReloading() override
+	{
+		return true;
+	}
+	//~ End IModuleInterface implementation
+};
+
+IMPLEMENT_MODULE(FPCGPythonInteropEditorModule, PCGPythonInteropEditor);
+
+DEFINE_LOG_CATEGORY(LogPCGPython);
\ No newline at end of file
diff --git a/Engine/Plugins/PCGInterops/PCGPythonInterop/Source/PCGPythonInteropEditor/Public/Elements/PCGExecutePythonScript.h b/Engine/Plugins/PCGInterops/PCGPythonInterop/Source/PCGPythonInteropEditor/Public/Elements/PCGExecutePythonScript.h
new file mode 100644
--- /dev/null
+++ b/Engine/Plugins/PCGInterops/PCGPythonInterop/Source/PCGPythonInteropEditor/Public/Elements/PCGExecutePythonScript.h
@@ -0,0 +1,115 @@
+// Copyright Epic Games, Inc. All Rights Reserved.
+
+#pragma once
+
+/** @todo_pcg: This node could welcome many improvements in the future, including:
+ * - EPythonCommandExecutionMode::EvaluateStatement can provide line-by-line feedback on python.
+ * - Optional param inputs/outputs to be get/set from within the Python script, similar to BP and HLSL.
+ * - Generalized source editor in PCG for HLSL, Python, and beyond.
+ */
+
+#include "PCGSettings.h"
+#include "Metadata/PCGDefaultValueInterface.h"
+
+#include "PCGExecutePythonScript.generated.h"
+
+UENUM()
+enum class EPCGPythonScriptInputMethod
+{
+	Input UMETA(Tooltip = "Execute the Python script from an input data or inline."),
+	File UMETA(Tooltip = "Execute the Python script from a selected file.")
+};
+
+/**
+* Execute a Python script, either inline, as input, or from a file.
+*/
+UCLASS(MinimalAPI, BlueprintType, ClassGroup = (Procedural))
+class UPCGExecutePythonScriptSettings : public UPCGSettings, public IPCGSettingsDefaultValueProvider
+{
+	GENERATED_BODY()
+
+	// The default script to employ until replaced by the user's script.
+	static const FString DefaultInlineScript;
+
+public:
+#if WITH_EDITOR
+	//~Begin UObject interface
+	virtual bool CanEditChange(const FProperty* InProperty) const override;
+	virtual EPCGChangeType GetChangeTypeForProperty(const FName& InPropertyName) const override;
+	//~End UObject interface
+
+	//~Begin UPCGSettings interface
+	virtual FName GetDefaultNodeName() const override;
+	virtual FText GetDefaultNodeTitle() const override;
+	virtual FText GetNodeTooltipText() const override;
+	virtual EPCGSettingsType GetType() const override;
+	virtual FString GetAdditionalTitleInformation() const override;
+	virtual bool CanCullTaskIfUnwired() const override { return false; }
+#endif // WITH_EDITOR
+
+protected:
+	virtual TArray<FPCGPinProperties> InputPinProperties() const override;
+	virtual TArray<FPCGPinProperties> OutputPinProperties() const override;
+	virtual FPCGElementPtr CreateElement() const override;
+	//~End UPCGSettings interface
+
+public:
+	//~Begin IPCGSettingsDefaultValueProvider interface
+	virtual bool DefaultValuesAreEnabled() const override;
+	virtual bool IsPinDefaultValueEnabled(FName PinLabel) const override;
+	virtual bool IsPinDefaultValueActivated(FName PinLabel) const override;
+	virtual EPCGMetadataTypes GetPinDefaultValueType(FName PinLabel) const override;
+	virtual bool IsPinDefaultValueMetadataTypeValid(FName PinLabel, EPCGMetadataTypes DataType) const override;
+	virtual bool CreateInitialDefaultValueAttribute(FName PinLabel, UPCGMetadata* OutMetadata) const override;
+	virtual bool IsInputPinRequiredByExecution(const UPCGPin* InPin) const override;
+
+#if WITH_EDITOR
+	virtual void ResetDefaultValues() override;
+	virtual void ResetDefaultValue(FName PinLabel) override;
+	virtual void SetPinDefaultValue(FName PinLabel, const FString& DefaultValue, bool bCreateIfNeeded = false) override;
+	virtual void SetPinDefaultValueIsActivated(FName PinLabel, bool bIsActivated, bool bDirtySettings = true) override;
+	virtual FString GetPinDefaultValueAsString(FName PinLabel) const override;
+	virtual FString GetPinInitialDefaultValueString(FName PinLabel) const override;
+	virtual EPCGSettingDefaultValueExtraFlags GetDefaultValueExtraFlags(FName PinLabel) const override;
+#endif // WITH_EDITOR
+
+protected:
+	virtual EPCGMetadataTypes GetPinInitialDefaultValueType(FName PinLabel) const override;
+	//~End IPCGSettingsDefaultValueProvider interface
+
+public:
+	/** The method for receiving the intended Python source. */
+	UPROPERTY(EditAnywhere, BlueprintReadWrite, Category = Settings)
+	EPCGPythonScriptInputMethod ScriptInputMethod = EPCGPythonScriptInputMethod::Input;
+
+	/** Which attribute to use as a script source. */
+	UPROPERTY(EditAnywhere, BlueprintReadWrite, Category = Settings, meta = (PCG_Overridable, EditCondition = "ScriptInputMethod == EPCGPythonScriptInputMethod::Input", EditConditionHides))
+	FPCGAttributePropertyInputSelector ScriptSource;
+
+	/** The path to the .py file that will be executed. */
+	UPROPERTY(EditAnywhere, BlueprintReadWrite, Category = Settings, meta = (PCG_Overridable, EditCondition = "ScriptInputMethod == EPCGPythonScriptInputMethod::File", EditConditionHides, FilePathFilter = "Python files (*.py)|*.py"))
+	FFilePath ScriptPath;
+
+	UPROPERTY(EditAnywhere, BlueprintReadWrite, Category = Settings, AdvancedDisplay, meta = (PCG_Overridable))
+	bool bMuteEditorToast = false;
+
+private:
+	// Inline constant script on the input pin to run.
+	UPROPERTY()
+	FString InlineScript = DefaultInlineScript;
+
+	// Inline constant is enabled for the input pin.
+	UPROPERTY()
+	bool bIsDefaultValueActivated = true;
+};
+
+class FPCGExecutePythonScriptElement : public IPCGElement
+{
+protected:
+	virtual bool ExecuteInternal(FPCGContext* InContext) const override;
+	// @todo_pcg: To be confirmed that python should be executed on main thread only.
+	virtual bool CanExecuteOnlyOnMainThread(FPCGContext* Context) const override { return true; }
+	// @todo_pcg: Could be a user option to run once only, etc
+	virtual bool IsCacheable(const UPCGSettings* InSettings) const override { return false; }
+};
+
diff --git a/Engine/Plugins/PCGInterops/PCGPythonInterop/Source/PCGPythonInteropEditor/Public/PCGPythonInteropEditorModule.h b/Engine/Plugins/PCGInterops/PCGPythonInterop/Source/PCGPythonInteropEditor/Public/PCGPythonInteropEditorModule.h
new file mode 100644
--- /dev/null
+++ b/Engine/Plugins/PCGInterops/PCGPythonInterop/Source/PCGPythonInteropEditor/Public/PCGPythonInteropEditorModule.h
@@ -0,0 +1,7 @@
+// Copyright Epic Games, Inc. All Rights Reserved.
+
+#pragma once
+
+#include "Logging/LogMacros.h"
+
+PCGPYTHONINTEROPEDITOR_API DECLARE_LOG_CATEGORY_EXTERN(LogPCGPython, Log, All);
\ No newline at end of file
